Added conditional ArenaMap compilation

diff --git a/Engine/Source/Programs/BenchmarkTool/Private/BenchmarkTool.cpp b/Engine/Source/Programs/BenchmarkTool/Private/BenchmarkTool.cpp
--- a/Engine/Source/Programs/BenchmarkTool/Private/BenchmarkTool.cpp
+++ b/Engine/Source/Programs/BenchmarkTool/Private/BenchmarkTool.cpp
@@ -317,6 +317,7 @@ inline FORCENOINLINE void ClobberMemory() { /* TODO */ }
 
 //////////////////////////////////////////////////////////////////////////
 
+#if UE_WITH_ARENAMAP
 void BM_MapPtrToArena(BenchmarkState& State)
 {
 	FArenaMap::SetRangeToArena(0, 32ull * 1024 * 1024 * 1024, nullptr);
@@ -330,6 +331,13 @@ void BM_MapPtrToArena(BenchmarkState& State)
 	}
 }
 
+UE_BENCHMARK(BM_MapPtrToArena)->Iterations(100000);
+UE_BENCHMARK(BM_MapPtrToArena)->Iterations(1000000);
+UE_BENCHMARK(BM_MapPtrToArena)->Iterations(10000000);
+UE_BENCHMARK(BM_MapPtrToArena)->Iterations(100000000);
+UE_BENCHMARK(BM_MapPtrToArena)->Iterations(1000000000);
+#endif
+
 void BM_NoOp(BenchmarkState& State)
 {
 	for (auto _ : State)
@@ -460,12 +468,6 @@ void BM_StdAtomicStoreRelaxed(BenchmarkState& State)
 UE_BENCHMARK(BM_NoOp)->Iterations(100000000);
 UE_BENCHMARK_CAPTURE(BM_NoOp, 1000, 1000)->Iterations(100000000);
 
-UE_BENCHMARK(BM_MapPtrToArena)->Iterations(100000);
-UE_BENCHMARK(BM_MapPtrToArena)->Iterations(1000000);
-UE_BENCHMARK(BM_MapPtrToArena)->Iterations(10000000);
-UE_BENCHMARK(BM_MapPtrToArena)->Iterations(100000000);
-UE_BENCHMARK(BM_MapPtrToArena)->Iterations(1000000000);
-
 UE_BENCHMARK(BM_CritSecLoad)->Iterations(100000000);
 
 UE_BENCHMARK(BM_TAtomic)->Iterations(100000000);
