Added RWLock/CriticalSection tests

diff --git a/Engine/Source/Programs/BenchmarkTool/Private/BenchmarkTool.cpp b/Engine/Source/Programs/BenchmarkTool/Private/BenchmarkTool.cpp
--- a/Engine/Source/Programs/BenchmarkTool/Private/BenchmarkTool.cpp
+++ b/Engine/Source/Programs/BenchmarkTool/Private/BenchmarkTool.cpp
@@ -482,6 +482,51 @@ UE_BENCHMARK(BM_StdAtomicRelaxed)->Iterations(100000000);
 UE_BENCHMARK(BM_StdAtomicStore)->Iterations(100000000);
 UE_BENCHMARK(BM_StdAtomicStoreRelaxed)->Iterations(100000000);
 
+//////////////////////////////////////////////////////////////////////////
+//
+// Basic tests to measure uncontended RWLock/Critical section performance
+//
+
+void BM_ReadWriteLock_ReadLock(BenchmarkState& State)
+{
+	FRWLock Lock;
+
+	for (auto _ : State)
+	{
+		Lock.ReadLock();
+		Lock.ReadUnlock();
+	}
+}
+
+void BM_ReadWriteLock_WriteLock(BenchmarkState& State)
+{
+	FRWLock Lock;
+
+	for (auto _ : State)
+	{
+		Lock.WriteLock();
+		Lock.WriteUnlock();
+	}
+}
+
+void BM_CriticalSection(BenchmarkState& State)
+{
+	FCriticalSection Lock;
+
+	for (auto _ : State)
+	{
+		Lock.Lock();
+		Lock.Unlock();
+	}
+}
+
+UE_BENCHMARK(BM_ReadWriteLock_ReadLock)->Iterations(10000000);
+UE_BENCHMARK(BM_ReadWriteLock_ReadLock)->Iterations(100000000);
+UE_BENCHMARK(BM_ReadWriteLock_WriteLock)->Iterations(10000000);
+UE_BENCHMARK(BM_ReadWriteLock_WriteLock)->Iterations(100000000);
+UE_BENCHMARK(BM_CriticalSection)->Iterations(10000000);
+UE_BENCHMARK(BM_CriticalSection)->Iterations(100000000);
+
 INT32_MAIN_INT32_ARGC_TCHAR_ARGV()
 {
 	GEngineLoop.PreInit(ArgC, ArgV);
