Added some shared pointer vs intrusive refcount benchmarks

diff --git a/Engine/Source/Programs/BenchmarkTool/Private/BenchmarkTool.cpp b/Engine/Source/Programs/BenchmarkTool/Private/BenchmarkTool.cpp
--- a/Engine/Source/Programs/BenchmarkTool/Private/BenchmarkTool.cpp
+++ b/Engine/Source/Programs/BenchmarkTool/Private/BenchmarkTool.cpp
@@ -4,6 +4,9 @@
 #include "BenchmarkTool.h"
 #include "Memory/MemoryArena.h"
 
+#include "Templates/RefCounting.h"
+#include "Templates/SharedPointer.h"
+
 #include "RequiredProgramMainCPPInclude.h"
 #include <locale.h>
 #include <xutility>
@@ -527,10 +530,92 @@ UE_BENCHMARK(BM_ReadWriteLock_WriteLock)->Iterations(100000000);
 UE_BENCHMARK(BM_CriticalSection)->Iterations(10000000);
 UE_BENCHMARK(BM_CriticalSection)->Iterations(100000000);
 
+//////////////////////////////////////////////////////////////////////////
+
+struct DummyShared
+{
+	int _ = 0;
+};
+
+void BM_TSharedPtr(BenchmarkState& State)
+{
+	for (auto _ : State)
+	{
+		TSharedPtr<DummyShared, ESPMode::ThreadSafe> Shared = MakeShared<DummyShared, ESPMode::ThreadSafe>();
+		DoNotOptimize(Shared);
+	}
+}
+
+void BM_TSharedPtrAssign(BenchmarkState& State)
+{
+	TSharedPtr<DummyShared, ESPMode::ThreadSafe> Shared = MakeShared<DummyShared, ESPMode::ThreadSafe>();
+
+	for (auto _ : State)
+	{
+		auto Shared2 = Shared;
+		DoNotOptimize(Shared2);
+	}
+}
+
+void BM_TSharedPtr_NoTS(BenchmarkState& State)
+{
+	for (auto _ : State)
+	{
+		auto Shared = MakeShared<DummyShared, ESPMode::NotThreadSafe>();
+		DoNotOptimize(Shared);
+	}
+}
+
+void BM_TSharedPtrAssign_NoTS(BenchmarkState& State)
+{
+	auto Shared = MakeShared<DummyShared, ESPMode::NotThreadSafe>();
+
+	for (auto _ : State)
+	{
+		auto Shared2 = Shared;
+		DoNotOptimize(Shared2);
+	}
+}
+
+struct DummyRefCount : public FRefCountBase
+{
+	int _ = 0;
+};
+
+void BM_TRefCountPtr(BenchmarkState& State)
+{
+	for (auto _ : State)
+	{
+		TRefCountPtr<DummyRefCount> RefCount = new DummyRefCount();
+		DoNotOptimize(RefCount);
+	}
+}
+
+void BM_TRefCountAssign(BenchmarkState& State)
+{
+	TRefCountPtr<DummyRefCount> RefCount = new DummyRefCount();
+
+	for (auto _ : State)
+	{
+		TRefCountPtr<DummyRefCount> Ref2 = RefCount;
+		DoNotOptimize(Ref2);
+	}
+}
+
+UE_BENCHMARK(BM_TSharedPtr)->Iterations(100000000);
+UE_BENCHMARK(BM_TRefCountPtr)->Iterations(100000000);
+UE_BENCHMARK(BM_TSharedPtr_NoTS)->Iterations(100000000);
+UE_BENCHMARK(BM_TSharedPtrAssign)->Iterations(100000000);
+UE_BENCHMARK(BM_TRefCountAssign)->Iterations(100000000);
+UE_BENCHMARK(BM_TSharedPtrAssign_NoTS)->Iterations(100000000);
+
+//////////////////////////////////////////////////////////////////////////
+
+#pragma comment(lib, "synchronization.lib")
+
 INT32_MAIN_INT32_ARGC_TCHAR_ARGV()
 {
 	GEngineLoop.PreInit(ArgC, ArgV);
-
 	BenchmarkRegistry::Get().RunBenchmarks();
 
 	return 0;
