Remove accidental #pragma comment

diff --git a/Engine/Source/Programs/BenchmarkTool/Private/BenchmarkTool.cpp b/Engine/Source/Programs/BenchmarkTool/Private/BenchmarkTool.cpp
--- a/Engine/Source/Programs/BenchmarkTool/Private/BenchmarkTool.cpp
+++ b/Engine/Source/Programs/BenchmarkTool/Private/BenchmarkTool.cpp
@@ -611,8 +611,6 @@ UE_BENCHMARK(BM_TSharedPtrAssign_NoTS)->Iterations(100000000);
 
 //////////////////////////////////////////////////////////////////////////
 
-#pragma comment(lib, "synchronization.lib")
-
 INT32_MAIN_INT32_ARGC_TCHAR_ARGV()
 {
 	GEngineLoop.PreInit(ArgC, ArgV);
