Add support to run QueuedThreadPool task on the task graph
Add support for benchmarking task graph scheduling in benchmark tool
Fix exit in benchmark tool

diff --git a/Engine/Source/Programs/BenchmarkTool/Private/BenchmarkTool.cpp b/Engine/Source/Programs/BenchmarkTool/Private/BenchmarkTool.cpp
--- a/Engine/Source/Programs/BenchmarkTool/Private/BenchmarkTool.cpp
+++ b/Engine/Source/Programs/BenchmarkTool/Private/BenchmarkTool.cpp
@@ -108,6 +108,7 @@ public:
 	{
 		UE_LOG(LogBenchmarkTool, Log, TEXT("Running '%s'..."), *this->Name);
 
+		TRACE_CPUPROFILER_EVENT_SCOPE_TEXT(*this->Name);
 		const uint64 StartTime = FPlatformTime::Cycles64();
 
 		Run(State);
@@ -201,6 +202,7 @@ public:
 
 	void RunBenchmarks()
 	{
+		TRACE_CPUPROFILER_EVENT_SCOPE(RunBenchmarks);
 		TArray<BenchmarkReporter::Run> RunResults;
 		RunResults.Reserve(Benchmarks.Num());
 
@@ -208,7 +210,7 @@ public:
 		FParse::Value(FCommandLine::Get(), TEXT("-Benchmark="), BenchName);
 		for (auto& Bench : Benchmarks)
 		{
-			if (BenchName.Len() > 0 && BenchName != Bench->Name)
+			if (BenchName.Len() > 0 && Bench->Name.Find(BenchName) == INDEX_NONE)
 			{
 				continue;
 			}
@@ -609,81 +611,117 @@ void BM_TRefCountAssign(BenchmarkState& State)
 	}
 }
 
-void BM_ThreadPoolOverhead_Impl(BenchmarkState& State, FQueuedThreadPool* ThreadPool)
+void BM_Scheduling_TaskGraphOverhead(BenchmarkState& State)
 {
-	struct FEventTask : public IQueuedWork
-	{
-		FEvent* Event;
-		FEventTask() { Event = FPlatformProcess::GetSynchEventFromPool(true); }
-		~FEventTask() { FPlatformProcess::ReturnSynchEventToPool(Event); }
-		void Abandon() override { }
-	};
+	FEvent* LastEvent = FPlatformProcess::GetSynchEventFromPool(true);
+	FEvent* WaitEvent = FPlatformProcess::GetSynchEventFromPool(true);
 
-	struct FWaitTask : public FEventTask
+	TAtomic<uint32> Count(0);
+	for (auto _ : State)
 	{
-		void DoThreadedWork() override { Event->Wait(); };
-	};
+		Count++;
+		FFunctionGraphTask::CreateAndDispatchWhenReady(
+			[WaitEvent, LastEvent, &Count](ENamedThreads::Type CurrentThread, const FGraphEventRef&CompletionGraphEvent)
+			{
+				TRACE_CPUPROFILER_EVENT_SCOPE(Task);
 
-	struct FTriggerTask : public FEventTask
-	{
-		void DoThreadedWork() override { Event->Trigger(); };
-	};
+				// Stall the tasks so we can benchmark the queuing code in AddQueuedWork.
+				// Otherwise, most threads will be able to execute as fast as the queuing
+				// happens and the dispatching will occur directly on each thread, 
+				// exercising a different code path.
+				WaitEvent->Wait();
+
+				if (--Count == 0)
+				{
+					LastEvent->Trigger();
+				}
+			},
+			QUICK_USE_CYCLE_STAT(BM_Scheduling_TaskGraphOverhead, STATGROUP_ThreadPoolAsyncTasks),
+			nullptr,
+			ENamedThreads::AnyThread
+			);
+	}
+
+	// Unstall the task processing so we can properly exercise the code
+	// path where all the threads need to pick another job to process.
+	WaitEvent->Trigger();
+
+	// Wait until the last task has executed
+	LastEvent->Wait();
+	
+	FPlatformProcess::ReturnSynchEventToPool(WaitEvent);
+	FPlatformProcess::ReturnSynchEventToPool(LastEvent);
+}
 
-	FWaitTask DefaultTask;
-	// Stall the tasks so we can benchmark the queuing code in AddQueuedWork.
-	// Otherwise, most threads will be able to execute as fast as the queuing
-	// happens and the dispatching will occur directly on each thread, 
-	// exercising a different code path.
-	DefaultTask.Event->Reset();
+void BM_Scheduling_ThreadPoolOverhead_Impl(BenchmarkState& State, FQueuedThreadPool* ThreadPool)
+{
+	FEvent* LastEvent = FPlatformProcess::GetSynchEventFromPool(true);
+	FEvent* WaitEvent = FPlatformProcess::GetSynchEventFromPool(true);
 
+	TAtomic<uint32> Count(0);
 	for (auto _ : State)
 	{
-		// Always enqueue the same empty task to avoid allocations.
-		ThreadPool->AddQueuedWork(&DefaultTask);
+		Count++;
+		AsyncPool(
+			*ThreadPool,
+			[WaitEvent, LastEvent, &Count]()
+			{
+				TRACE_CPUPROFILER_EVENT_SCOPE(Task);
+
+				// Stall the tasks so we can benchmark the queuing code in AddQueuedWork.
+				// Otherwise, most threads will be able to execute as fast as the queuing
+				// happens and the dispatching will occur directly on each thread, 
+				// exercising a different code path.
+				WaitEvent->Wait();
+
+				if (--Count == 0)
+				{
+					LastEvent->Trigger();
+				}
+			}
+		);
 	}
 
 	// Unstall the task processing so we can properly exercise the code
 	// path where all the threads need to pick another job to process.
-	// (i.e. ReturnToPoolOrGetNextJob)
-	DefaultTask.Event->Trigger();
-
-	// Enqueue a sentinel task so we know when everything has been unqueued
-	FTriggerTask LastTask;
-	ThreadPool->AddQueuedWork(&LastTask);
-	LastTask.Event->Wait();
-
-	// If run with more than 1 thread, there might be up to (number of threads) empty task 
-	// still being processed at this point, but since work is not meaningful we don't really care.
-	// The thread pool destructor will abandon any task left and clear up threads before going
-	// out of scope so we're safe to exit without any additional cleanup.
+	WaitEvent->Trigger();
+
+	// Wait until the last task has executed
+	LastEvent->Wait();
+	
+	FPlatformProcess::ReturnSynchEventToPool(WaitEvent);
+	FPlatformProcess::ReturnSynchEventToPool(LastEvent);
 }
 
 // This test is probably only meaningful when comparing relative speed
 // of threadpool implementations and to profile the current one.
-void BM_ThreadPoolOverhead(BenchmarkState& State)
+void BM_Scheduling_ThreadPoolOverhead(BenchmarkState& State)
 {
-	// A single thread is enough to exercise the Queuing and Dequeuing code paths.
-	// More threads will only increase contention as the current implementation of
-	// those functions is under critical sections anyway...
 	TUniquePtr<FQueuedThreadPool> ThreadPool(FQueuedThreadPool::Allocate());
-	check(ThreadPool && ThreadPool->Create(1));
+	check(ThreadPool && ThreadPool->Create(FPlatformMisc::NumberOfCores()));
 
-	BM_ThreadPoolOverhead_Impl(State, ThreadPool.Get());
+	BM_Scheduling_ThreadPoolOverhead_Impl(State, ThreadPool.Get());
 }
 
 // This test is probably only meaningful when comparing relative speed
 // of threadpool implementations and to profile the current one.
-void BM_ThreadPoolWrapperOverhead(BenchmarkState& State)
+void BM_Scheduling_ThreadPoolWrapperOverhead(BenchmarkState& State)
 {
-	// A single thread is enough to exercise the Queuing and Dequeuing code paths.
-	// More threads will only increase contention as the current implementation of
-	// those functions is under critical sections anyway...
 	TUniquePtr<FQueuedThreadPool> ThreadPool(FQueuedThreadPool::Allocate());
-	check(ThreadPool && ThreadPool->Create(1));
+	check(ThreadPool && ThreadPool->Create(FPlatformMisc::NumberOfCores()));
 
 	FQueuedThreadPoolWrapper ThreadPoolWrapper(ThreadPool.Get());
 
-	BM_ThreadPoolOverhead_Impl(State, &ThreadPoolWrapper);
+	BM_Scheduling_ThreadPoolOverhead_Impl(State, &ThreadPoolWrapper);
+}
+
+// This test is probably only meaningful when comparing relative speed
+// of threadpool implementations and to profile the current one.
+void BM_Scheduling_ThreadPoolTaskGraphWrapperOverhead(BenchmarkState& State)
+{
+	FQueuedThreadPoolTaskGraphWrapper ThreadPoolWrapper(ENamedThreads::AnyThread);
+
+	BM_Scheduling_ThreadPoolOverhead_Impl(State, &ThreadPoolWrapper);
 }
 
 UE_BENCHMARK(BM_TSharedPtr)->Iterations(100000000);
@@ -692,8 +730,13 @@ UE_BENCHMARK(BM_TSharedPtr_NoTS)->Iterations(100000000);
 UE_BENCHMARK(BM_TSharedPtrAssign)->Iterations(100000000);
 UE_BENCHMARK(BM_TRefCountAssign)->Iterations(100000000);
 UE_BENCHMARK(BM_TSharedPtrAssign_NoTS)->Iterations(100000000);
-UE_BENCHMARK(BM_ThreadPoolOverhead)->Iterations(100000);
-UE_BENCHMARK(BM_ThreadPoolWrapperOverhead)->Iterations(100000);
+
+// You can compare all scheduling tasks by using this command line -Benchmark=BM_Scheduling
+// You can also test scalability of the different schedulers by adding -corelimit= to the command line
+UE_BENCHMARK(BM_Scheduling_ThreadPoolOverhead)->Iterations(100000);
+UE_BENCHMARK(BM_Scheduling_ThreadPoolWrapperOverhead)->Iterations(100000);
+UE_BENCHMARK(BM_Scheduling_TaskGraphOverhead)->Iterations(100000);
+UE_BENCHMARK(BM_Scheduling_ThreadPoolTaskGraphWrapperOverhead)->Iterations(100000);
 
 //////////////////////////////////////////////////////////////////////////
 
@@ -702,5 +745,7 @@ INT32_MAIN_INT32_ARGC_TCHAR_ARGV()
 	GEngineLoop.PreInit(ArgC, ArgV);
 	BenchmarkRegistry::Get().RunBenchmarks();
 
+	FEngineLoop::AppPreExit();
+	FEngineLoop::AppExit();
 	return 0;
-}
+}
\ No newline at end of file
diff --git a/Engine/Source/Runtime/Core/Public/Misc/QueuedThreadPoolWrapper.h b/Engine/Source/Runtime/Core/Public/Misc/QueuedThreadPoolWrapper.h
--- a/Engine/Source/Runtime/Core/Public/Misc/QueuedThreadPoolWrapper.h
+++ b/Engine/Source/Runtime/Core/Public/Misc/QueuedThreadPoolWrapper.h
@@ -6,6 +6,8 @@
 #include "GenericPlatform/GenericPlatformAffinity.h"
 #include "QueuedThreadPool.h"
 #include "ScopeRWLock.h"
+#include "Async/TaskGraphInterfaces.h"
+#include "Misc/IQueuedWork.h"
 
 /** ThreadPool wrapper implementation allowing to schedule
   * up to MaxConcurrency tasks at a time making sub-partitioning
@@ -97,4 +99,107 @@ public:
 		FRWScopeLock ScopeLock(Lock, SLT_Write);
 		QueuedWork.Sort(EQueuedWorkPriority::Normal, Predicate);
 	}
+};
+
+/** ThreadPool wrapper implementation allowing to schedule thread-pool tasks on the task graph.
+  */
+class CORE_API FQueuedThreadPoolTaskGraphWrapper : public FQueuedThreadPool
+{
+public:
+	/**
+	 * InPriorityMapper           Thread-safe function used to map any priority from this Queue to the priority that should be used when scheduling the task on the task graph.
+	 */
+	FQueuedThreadPoolTaskGraphWrapper(TFunction<ENamedThreads::Type (EQueuedWorkPriority)> InPriorityMapper = nullptr)
+		: TaskCount(0)
+		, bIsExiting(0)
+	{
+		if (InPriorityMapper)
+		{
+			PriorityMapper = InPriorityMapper;
+		}
+		else
+		{
+			PriorityMapper = [this](EQueuedWorkPriority InPriority) { return GetDefaultPriorityMapping(InPriority); };
+		}
+	}
+
+	/**
+	 * InDesiredThread           The task-graph desired thread and priority.
+	 */
+	FQueuedThreadPoolTaskGraphWrapper(ENamedThreads::Type InDesiredThread)
+		: TaskCount(0)
+		, bIsExiting(0)
+	{
+		PriorityMapper = [InDesiredThread](EQueuedWorkPriority InPriority) { return InDesiredThread; };
+	}
+
+	~FQueuedThreadPoolTaskGraphWrapper()
+	{
+		Destroy();
+	}
+private:
+	void AddQueuedWork(IQueuedWork* InQueuedWork, EQueuedWorkPriority InPriority = EQueuedWorkPriority::Normal) override
+	{
+		check(bIsExiting == false);
+		TaskCount++;
+		FFunctionGraphTask::CreateAndDispatchWhenReady(
+			[this, InQueuedWork](ENamedThreads::Type CurrentThread, const FGraphEventRef& MyCompletionGraphEvent)
+			{
+				FMemMark Mark(FMemStack::Get());
+				InQueuedWork->DoThreadedWork();
+				OnTaskCompleted(InQueuedWork);
+			},
+			QUICK_USE_CYCLE_STAT(FQueuedThreadPoolTaskGraphWrapper, STATGROUP_ThreadPoolAsyncTasks),
+			nullptr,
+			PriorityMapper(InPriority)
+		);
+	}
+
+	bool RetractQueuedWork(IQueuedWork* InQueuedWork) override
+	{
+		// The task graph doesn't support retraction for now
+		return false;
+	}
+
+	void OnTaskCompleted(IQueuedWork* InQueuedWork)
+	{
+		--TaskCount;
+	}
+
+	int32 GetNumThreads() const override
+	{
+		return FTaskGraphInterface::Get().GetNumWorkerThreads();
+	}
+
+	ENamedThreads::Type GetDefaultPriorityMapping(EQueuedWorkPriority InQueuedWorkPriority)
+	{
+		ENamedThreads::Type DesiredThread = ENamedThreads::AnyNormalThreadNormalTask;
+		if (InQueuedWorkPriority > EQueuedWorkPriority::Normal)
+		{
+			DesiredThread = ENamedThreads::AnyBackgroundThreadNormalTask;
+		}
+		else if (InQueuedWorkPriority < EQueuedWorkPriority::Normal)
+		{
+			DesiredThread = ENamedThreads::AnyHiPriThreadNormalTask;
+		}
+		return DesiredThread;
+	}
+protected:
+	bool Create(uint32 InNumQueuedThreads, uint32 StackSize, EThreadPriority ThreadPriority, const TCHAR* Name) override
+	{
+		return true;
+	}
+
+	void Destroy() override
+	{
+		bIsExiting = true;
+		while (TaskCount != 0)
+		{
+			FPlatformProcess::Sleep(0.01f);
+		}
+	}
+private:
+	TFunction<ENamedThreads::Type (EQueuedWorkPriority)> PriorityMapper;
+	TAtomic<uint32> TaskCount;
+	TAtomic<bool> bIsExiting;
 };
\ No newline at end of file
