Deprecated FQueuedThreadPoolTaskGraphWrapper class & removed from Benchmark tool since its no longer being used. #jira nojira

#rnx

#Synced-CL 48675545

diff --git a/Engine/Source/Programs/BenchmarkTool/Private/BenchmarkTool.cpp b/Engine/Source/Programs/BenchmarkTool/Private/BenchmarkTool.cpp
--- a/Engine/Source/Programs/BenchmarkTool/Private/BenchmarkTool.cpp
+++ b/Engine/Source/Programs/BenchmarkTool/Private/BenchmarkTool.cpp
@@ -715,15 +715,6 @@ void BM_Scheduling_ThreadPoolWrapperOverhead(BenchmarkState& State)
 	BM_Scheduling_ThreadPoolOverhead_Impl(State, &ThreadPoolWrapper);
 }
 
-// This test is probably only meaningful when comparing relative speed
-// of threadpool implementations and to profile the current one.
-void BM_Scheduling_ThreadPoolTaskGraphWrapperOverhead(BenchmarkState& State)
-{
-	FQueuedThreadPoolTaskGraphWrapper ThreadPoolWrapper(ENamedThreads::AnyThread);
-
-	BM_Scheduling_ThreadPoolOverhead_Impl(State, &ThreadPoolWrapper);
-}
-
 // This test is probably only meaningful when comparing relative speed
 // of threadpool implementations and to profile the current one.
 void BM_Scheduling_ThreadPoolLowLevelWrapperOverhead(BenchmarkState& State)
@@ -745,7 +736,6 @@ UE_BENCHMARK(BM_TSharedPtrAssign_NoTS)->Iterations(100000000);
 UE_BENCHMARK(BM_Scheduling_ThreadPoolOverhead)->Iterations(100000);
 UE_BENCHMARK(BM_Scheduling_ThreadPoolWrapperOverhead)->Iterations(100000);
 UE_BENCHMARK(BM_Scheduling_TaskGraphOverhead)->Iterations(100000);
-UE_BENCHMARK(BM_Scheduling_ThreadPoolTaskGraphWrapperOverhead)->Iterations(100000);
 UE_BENCHMARK(BM_Scheduling_ThreadPoolLowLevelWrapperOverhead)->Iterations(100000);
 
 //////////////////////////////////////////////////////////////////////////
diff --git a/Engine/Source/Runtime/Core/Public/Misc/QueuedThreadPoolWrapper.h b/Engine/Source/Runtime/Core/Public/Misc/QueuedThreadPoolWrapper.h
--- a/Engine/Source/Runtime/Core/Public/Misc/QueuedThreadPoolWrapper.h
+++ b/Engine/Source/Runtime/Core/Public/Misc/QueuedThreadPoolWrapper.h
@@ -263,7 +263,7 @@ public:
 
 /** ThreadPool wrapper implementation allowing to schedule thread-pool tasks on the task graph.
   */
-class FQueuedThreadPoolTaskGraphWrapper final : public FQueuedThreadPool
+class UE_DEPRECATED(5.8, "FQueuedThreadPoolTaskGraphWrapper is obsolete and will be removed") FQueuedThreadPoolTaskGraphWrapper final : public FQueuedThreadPool
 {
 public:
 	/**
